PR #5127: Minor changes to localization dashboard (Contributed by [REDACTED_PERSON])

diff --git a/Engine/Source/Editor/LocalizationCommandletExecution/Private/LocalizationCommandletExecution.cpp b/Engine/Source/Editor/LocalizationCommandletExecution/Private/LocalizationCommandletExecution.cpp
--- a/Engine/Source/Editor/LocalizationCommandletExecution/Private/LocalizationCommandletExecution.cpp
+++ b/Engine/Source/Editor/LocalizationCommandletExecution/Private/LocalizationCommandletExecution.cpp
@@ -698,7 +698,7 @@ namespace
 
 	FText SLocalizationCommandletExecutor::GetCloseButtonText() const
 	{
-		return HasCompleted() ? LOCTEXT("OkayButtonText", "Okay") : LOCTEXT("CancelButtonText", "Cancel");
+		return HasCompleted() ? LOCTEXT("OkButtonText", "OK") : LOCTEXT("CancelButtonText", "Cancel");
 	}
 
 	FReply SLocalizationCommandletExecutor::OnCloseButtonClicked()
