PR #5980: Localization commandlet executor fixes (Contributed by [REDACTED_PERSON])

diff --git a/Engine/Source/Editor/LocalizationCommandletExecution/Private/LocalizationCommandletExecution.cpp b/Engine/Source/Editor/LocalizationCommandletExecution/Private/LocalizationCommandletExecution.cpp
--- a/Engine/Source/Editor/LocalizationCommandletExecution/Private/LocalizationCommandletExecution.cpp
+++ b/Engine/Source/Editor/LocalizationCommandletExecution/Private/LocalizationCommandletExecution.cpp
@@ -257,7 +257,6 @@ namespace
 							[
 								SNew(SButton)
 								.ContentPadding(FMargin(6.0f, 2.0f))
-								.IsEnabled(false)
 								.Text(LOCTEXT("SaveLogButtonText", "Save Log..."))
 								.ToolTipText(LOCTEXT("SaveLogButtonToolTip", "Save the logged text to a file."))
 								.OnClicked(this, &SLocalizationCommandletExecutor::OnSaveLogClicked)
@@ -280,6 +279,11 @@ namespace
 		{
 			CurrentTaskIndex = 0;
 			ExecuteCommandlet(TaskListModels[CurrentTaskIndex].ToSharedRef());
+
+			if (TaskListView.IsValid())
+			{
+				TaskListView->SetSelection(TaskListModels[CurrentTaskIndex]);
+			}
 		}
 	}
 
@@ -360,6 +364,11 @@ namespace
 			{
 				CurrentTaskModel = TaskListModels[CurrentTaskIndex];
 				ExecuteCommandlet(CurrentTaskModel.ToSharedRef());
+
+				if (TaskListView.IsValid())
+				{
+					TaskListView->SetSelection(TaskListModels[CurrentTaskIndex]);
+				}
 			}
 		}
 		// Non-zero is a failure.
