Fixed the PortableObject data source not detecting newly added localization targets

#rnx

diff --git a/Engine/Plugins/Editor/ContentBrowser/ContentBrowserFileDataSource/Source/ContentBrowserFileDataSource/Private/ContentBrowserFileDataSource.cpp b/Engine/Plugins/Editor/ContentBrowser/ContentBrowserFileDataSource/Source/ContentBrowserFileDataSource/Private/ContentBrowserFileDataSource.cpp
--- a/Engine/Plugins/Editor/ContentBrowser/ContentBrowserFileDataSource/Source/ContentBrowserFileDataSource/Private/ContentBrowserFileDataSource.cpp
+++ b/Engine/Plugins/Editor/ContentBrowser/ContentBrowserFileDataSource/Source/ContentBrowserFileDataSource/Private/ContentBrowserFileDataSource.cpp
@@ -533,6 +533,11 @@ void UContentBrowserFileDataSource::RemoveFileMount(const FName InFileMountPath)
 	}
 }
 
+bool UContentBrowserFileDataSource::HasFileMount(const FName InFileMountPath) const
+{
+	return RegisteredFileMounts.Contains(InFileMountPath);
+}
+
 void UContentBrowserFileDataSource::Tick(const float InDeltaTime)
 {
 	Super::Tick(InDeltaTime);
diff --git a/Engine/Plugins/Editor/ContentBrowser/ContentBrowserFileDataSource/Source/ContentBrowserFileDataSource/Public/ContentBrowserFileDataSource.h b/Engine/Plugins/Editor/ContentBrowser/ContentBrowserFileDataSource/Source/ContentBrowserFileDataSource/Public/ContentBrowserFileDataSource.h
--- a/Engine/Plugins/Editor/ContentBrowser/ContentBrowserFileDataSource/Source/ContentBrowserFileDataSource/Public/ContentBrowserFileDataSource.h
+++ b/Engine/Plugins/Editor/ContentBrowser/ContentBrowserFileDataSource/Source/ContentBrowserFileDataSource/Public/ContentBrowserFileDataSource.h
@@ -41,6 +41,8 @@ public:
 
 	void RemoveFileMount(const FName InFileMountPath);
 
+	bool HasFileMount(const FName InFileMountPath) const;
+
 	virtual void Tick(const float InDeltaTime);
 
 	virtual void CompileFilter(const FName InPath, const FContentBrowserDataFilter& InFilter, FContentBrowserDataCompiledFilter& OutCompiledFilter) override;
diff --git a/Engine/Plugins/Editor/Localization/PortableObjectFileDataSource/Source/PortableObjectFileDataSource/PortableObjectFileDataSource.Build.cs b/Engine/Plugins/Editor/Localization/PortableObjectFileDataSource/Source/PortableObjectFileDataSource/PortableObjectFileDataSource.Build.cs
--- a/Engine/Plugins/Editor/Localization/PortableObjectFileDataSource/Source/PortableObjectFileDataSource/PortableObjectFileDataSource.Build.cs
+++ b/Engine/Plugins/Editor/Localization/PortableObjectFileDataSource/Source/PortableObjectFileDataSource/PortableObjectFileDataSource.Build.cs
@@ -19,6 +19,7 @@ namespace UnrealBuildTool.Rules
 				new string[]
 				{
 					"ContentBrowserFileDataSource",
+					"Localization",
 				}
 			);
 		}
diff --git a/Engine/Plugins/Editor/Localization/PortableObjectFileDataSource/Source/PortableObjectFileDataSource/Private/PortableObjectFileDataSourceModule.cpp b/Engine/Plugins/Editor/Localization/PortableObjectFileDataSource/Source/PortableObjectFileDataSource/Private/PortableObjectFileDataSourceModule.cpp
--- a/Engine/Plugins/Editor/Localization/PortableObjectFileDataSource/Source/PortableObjectFileDataSource/Private/PortableObjectFileDataSourceModule.cpp
+++ b/Engine/Plugins/Editor/Localization/PortableObjectFileDataSource/Source/PortableObjectFileDataSource/Private/PortableObjectFileDataSourceModule.cpp
@@ -2,6 +2,7 @@
 
 #include "ContentBrowserFileDataCore.h"
 #include "ContentBrowserFileDataSource.h"
+#include "LocalizationDelegates.h"
 #include "Misc/PackageName.h"
 #include "Misc/Paths.h"
 #include "Modules/ModuleManager.h"
@@ -89,6 +90,7 @@ public:
 			// Listen for new paths that may contain localization data
 			FPackageName::OnContentPathMounted().AddRaw(this, &FPortableObjectFileDataSourceModule::OnContentPathMounted);
 			FPackageName::OnContentPathDismounted().AddRaw(this, &FPortableObjectFileDataSourceModule::OnContentPathDismounted);
+			LocalizationDelegates::OnLocalizationTargetDataUpdated.AddRaw(this, &FPortableObjectFileDataSourceModule::OnLocalizationTargetDataUpdated);
 		}
 	}
 
@@ -96,6 +98,7 @@ public:
 	{
 		FPackageName::OnContentPathMounted().RemoveAll(this);
 		FPackageName::OnContentPathDismounted().RemoveAll(this);
+		LocalizationDelegates::OnLocalizationTargetDataUpdated.RemoveAll(this);
 			
 		PoFileDataSource.Reset();
 	}
@@ -121,6 +124,26 @@ private:
 		}
 	}
 
+	void OnLocalizationTargetDataUpdated(const FString& InLocalizationTargetPath)
+	{
+		if (PoFileDataSource)
+		{
+			// Trim the target name from the path so that we're querying for the outer "Localization" folder that hosts the target(s)
+			const FString LocalizationTargetFolder = FPaths::GetPath(InLocalizationTargetPath);
+
+			FString LocalizationTargetAssetPath;
+			if (FPaths::DirectoryExists(LocalizationTargetFolder) && FPackageName::TryConvertFilenameToLongPackageName(LocalizationTargetFolder, LocalizationTargetAssetPath))
+			{
+				// Updating this target may have added a "Localization" folder that we're not currently monitoring
+				// If so, add that mount now (any existing mounts will be updated via the directory watcher)
+				if (!PoFileDataSource->HasFileMount(*LocalizationTargetAssetPath))
+				{
+					PoFileDataSource->AddFileMount(*LocalizationTargetAssetPath, LocalizationTargetFolder);
+				}
+			}
+		}
+	}
+
 	TStrongObjectPtr<UContentBrowserFileDataSource> PoFileDataSource;
 };
 
diff --git a/Engine/Source/Developer/Localization/Private/LocalizationDelegates.cpp b/Engine/Source/Developer/Localization/Private/LocalizationDelegates.cpp
new file mode 100644
--- /dev/null
+++ b/Engine/Source/Developer/Localization/Private/LocalizationDelegates.cpp
@@ -0,0 +1,5 @@
+// Copyright Epic Games, Inc. All Rights Reserved.
+
+#include "LocalizationDelegates.h"
+
+TMulticastDelegate<void(const FString& LocalizationTargetPath)> LocalizationDelegates::OnLocalizationTargetDataUpdated;
diff --git a/Engine/Source/Developer/Localization/Private/UserGeneratedContentLocalization.cpp b/Engine/Source/Developer/Localization/Private/UserGeneratedContentLocalization.cpp
--- a/Engine/Source/Developer/Localization/Private/UserGeneratedContentLocalization.cpp
+++ b/Engine/Source/Developer/Localization/Private/UserGeneratedContentLocalization.cpp
@@ -15,6 +15,7 @@
 #include "Internationalization/TextLocalizationManager.h"
 #include "TextLocalizationResourceGenerator.h"
 #include "LocalizationConfigurationScript.h"
+#include "LocalizationDelegates.h"
 #include "LocTextHelper.h"
 #include "PortableObjectFormatDOM.h"
 #include "SourceControlHelpers.h"
@@ -497,6 +498,8 @@ bool ExportLocalization(TArrayView<const TSharedRef<IPlugin>> Plugins, const FEx
 				}
 			}
 		}
+
+		LocalizationDelegates::OnLocalizationTargetDataUpdated.Broadcast(PluginLocalizationTargetDirectory);
 	}
 
 	return true;
@@ -547,6 +550,8 @@ bool CompileLocalizationTarget(const FString& LocalizationTargetDirectory, const
 		}
 	}
 
+	LocalizationDelegates::OnLocalizationTargetDataUpdated.Broadcast(LocalizationTargetDirectory);
+
 	return true;
 }
 
diff --git a/Engine/Source/Developer/Localization/Public/LocalizationDelegates.h b/Engine/Source/Developer/Localization/Public/LocalizationDelegates.h
new file mode 100644
--- /dev/null
+++ b/Engine/Source/Developer/Localization/Public/LocalizationDelegates.h
@@ -0,0 +1,17 @@
+// Copyright Epic Games, Inc. All Rights Reserved.
+
+#pragma once
+
+#include "Delegates/Delegate.h"
+
+class FString;
+
+namespace LocalizationDelegates
+{
+
+/**
+ * Delegate called when any on-disk data for the given localization target is updated.
+ */
+LOCALIZATION_API extern TMulticastDelegate<void(const FString& LocalizationTargetPath)> OnLocalizationTargetDataUpdated;
+
+}
diff --git a/Engine/Source/Editor/LocalizationCommandletExecution/Private/LocalizationCommandletTasks.cpp b/Engine/Source/Editor/LocalizationCommandletExecution/Private/LocalizationCommandletTasks.cpp
--- a/Engine/Source/Editor/LocalizationCommandletExecution/Private/LocalizationCommandletTasks.cpp
+++ b/Engine/Source/Editor/LocalizationCommandletExecution/Private/LocalizationCommandletTasks.cpp
@@ -9,11 +9,46 @@
 #include "Sound/DialogueWave.h"
 #include "LocalizationCommandletExecution.h"
 #include "LocalizationConfigurationScript.h"
+#include "LocalizationDelegates.h"
 #include "LocalizationTargetTypes.h"
 #include "HAL/FileManager.h"
 
 #define LOCTEXT_NAMESPACE "LocalizationCommandletTasks"
 
+namespace LocalizationCommandletTasks
+{
+
+void BroadcastLocalizationTargetsDataUpdated(TArrayView<ULocalizationTarget* const> Targets)
+{
+	for (ULocalizationTarget* Target : Targets)
+	{
+		const FString TargetContentPath = FPaths::Combine(Target->IsMemberOfEngineTargetSet() ? FPaths::EngineContentDir() : FPaths::ProjectContentDir(), TEXT("Localization"), Target->Settings.Name);
+		LocalizationDelegates::OnLocalizationTargetDataUpdated.Broadcast(TargetContentPath);
+	}
+}
+
+bool ExecuteTasksAndBroadcastLocalizationTargetsDataUpdated(const TSharedRef<SWindow>& ParentWindow, const FText& Title, const TArray<LocalizationCommandletExecution::FTask>& Tasks, TArrayView<ULocalizationTarget* const> Targets)
+{
+	if (LocalizationCommandletExecution::Execute(ParentWindow, Title, Tasks))
+	{
+		BroadcastLocalizationTargetsDataUpdated(Targets);
+		return true;
+	}
+	return false;
+}
+
+bool ExecuteTasksAndBroadcastLocalizationTargetDataUpdated(const TSharedRef<SWindow>& ParentWindow, const FText& Title, const TArray<LocalizationCommandletExecution::FTask>& Tasks, ULocalizationTarget* const Target)
+{
+	if (LocalizationCommandletExecution::Execute(ParentWindow, Title, Tasks))
+	{
+		BroadcastLocalizationTargetsDataUpdated(MakeArrayView(&Target, 1));
+		return true;
+	}
+	return false;
+}
+
+}
+
 bool LocalizationCommandletTasks::GatherTextForTargets(const TSharedRef<SWindow>& ParentWindow, const TArray<ULocalizationTarget*>& Targets)
 {
 	TArray<LocalizationCommandletExecution::FTask> Tasks;
@@ -31,7 +66,7 @@ bool LocalizationCommandletTasks::GatherTextForTargets(const TSharedRef<SWindow>
 		Tasks.Add(LocalizationCommandletExecution::FTask(GatherTaskName, GatherScriptPath, ShouldUseProjectFile));
 	}
 
-	return LocalizationCommandletExecution::Execute(ParentWindow, LOCTEXT("GatherAllTargetsWindowTitle", "Gather Text for All Targets"), Tasks);
+	return ExecuteTasksAndBroadcastLocalizationTargetsDataUpdated(ParentWindow, LOCTEXT("GatherAllTargetsWindowTitle", "Gather Text for All Targets"), Tasks, Targets);
 }
 
 bool LocalizationCommandletTasks::GatherTextForTarget(const TSharedRef<SWindow>& ParentWindow, ULocalizationTarget* const Target)
@@ -47,7 +82,7 @@ bool LocalizationCommandletTasks::GatherTextForTarget(const TSharedRef<SWindow>&
 	Arguments.Add(TEXT("TargetName"), FText::FromString(Target->Settings.Name));
 
 	const FText WindowTitle = FText::Format(LOCTEXT("GatherTargetWindowTitle", "Gather Text for Target {TargetName}"), Arguments);
-	return LocalizationCommandletExecution::Execute(ParentWindow, WindowTitle, Tasks);
+	return ExecuteTasksAndBroadcastLocalizationTargetDataUpdated(ParentWindow, WindowTitle, Tasks, Target);
 }
 
 
@@ -74,7 +109,7 @@ bool LocalizationCommandletTasks::ImportTextForTargets(const TSharedRef<SWindow>
 		Tasks.Add(LocalizationCommandletExecution::FTask(ReportTaskName, ReportScriptPath, ShouldUseProjectFile));
 	}
 
-	return LocalizationCommandletExecution::Execute(ParentWindow, LOCTEXT("ImportForAllTargetsWindowTitle", "Import Translations for All Targets"), Tasks);
+	return ExecuteTasksAndBroadcastLocalizationTargetsDataUpdated(ParentWindow, LOCTEXT("ImportForAllTargetsWindowTitle", "Import Translations for All Targets"), Tasks, Targets);
 }
 
 bool LocalizationCommandletTasks::ImportTextForTarget(const TSharedRef<SWindow>& ParentWindow, ULocalizationTarget* const Target, const TOptional<FString> DirectoryPath)
@@ -94,7 +129,7 @@ bool LocalizationCommandletTasks::ImportTextForTarget(const TSharedRef<SWindow>&
 	Arguments.Add(TEXT("TargetName"), FText::FromString(Target->Settings.Name));
 
 	const FText WindowTitle = FText::Format(LOCTEXT("ImportForTargetWindowTitle", "Import Translations for Target {TargetName}"), Arguments);
-	return LocalizationCommandletExecution::Execute(ParentWindow, WindowTitle, Tasks);
+	return ExecuteTasksAndBroadcastLocalizationTargetDataUpdated(ParentWindow, WindowTitle, Tasks, Target);
 }
 
 bool LocalizationCommandletTasks::ImportTextForCulture(const TSharedRef<SWindow>& ParentWindow, ULocalizationTarget* const Target, const FString& CultureName, const TOptional<FString> FilePath)
@@ -123,7 +158,7 @@ bool LocalizationCommandletTasks::ImportTextForCulture(const TSharedRef<SWindow>
 
 	const FText WindowTitle = FText::Format(LOCTEXT("ImportCultureForTargetWindowTitle", "Import {CultureName} Translations for Target {TargetName}"), Arguments);
 
-	bool HasSucceeeded = LocalizationCommandletExecution::Execute(ParentWindow, WindowTitle, Tasks);
+	bool HasSucceeeded = ExecuteTasksAndBroadcastLocalizationTargetDataUpdated(ParentWindow, WindowTitle, Tasks, Target);
 	IFileManager::Get().Delete(*ImportScriptPath); // Don't clutter up the loc config directory with scripts for individual cultures.
 	return HasSucceeeded;
 }
@@ -146,7 +181,7 @@ bool LocalizationCommandletTasks::ExportTextForTargets(const TSharedRef<SWindow>
 		Tasks.Add(LocalizationCommandletExecution::FTask(ExportTaskName, ExportScriptPath, ShouldUseProjectFile));
 	}
 
-	return LocalizationCommandletExecution::Execute(ParentWindow, LOCTEXT("ExportForAllTargetsWindowTitle", "Export Translations for All Targets"), Tasks);
+	return ExecuteTasksAndBroadcastLocalizationTargetsDataUpdated(ParentWindow, LOCTEXT("ExportForAllTargetsWindowTitle", "Export Translations for All Targets"), Tasks, Targets);
 }
 
 bool LocalizationCommandletTasks::ExportTextForTarget(const TSharedRef<SWindow>& ParentWindow, ULocalizationTarget* const Target, const TOptional<FString> DirectoryPath)
@@ -162,7 +197,7 @@ bool LocalizationCommandletTasks::ExportTextForTarget(const TSharedRef<SWindow>&
 	Arguments.Add(TEXT("TargetName"), FText::FromString(Target->Settings.Name));
 
 	const FText WindowTitle = FText::Format(LOCTEXT("ExportForTargetWindowTitle", "Export Translations for Target {TargetName}"), Arguments);
-	return LocalizationCommandletExecution::Execute(ParentWindow, WindowTitle, Tasks);
+	return ExecuteTasksAndBroadcastLocalizationTargetDataUpdated(ParentWindow, WindowTitle, Tasks, Target);
 }
 
 bool LocalizationCommandletTasks::ExportTextForCulture(const TSharedRef<SWindow>& ParentWindow, ULocalizationTarget* const Target, const FString& CultureName, const TOptional<FString> FilePath)
@@ -187,7 +222,7 @@ bool LocalizationCommandletTasks::ExportTextForCulture(const TSharedRef<SWindow>
 
 	const FText WindowTitle = FText::Format(LOCTEXT("ExportCultureForTargetWindowTitle", "Export {CultureName} Translations for Target {TargetName}"), Arguments);
 
-	bool HasSucceeeded = LocalizationCommandletExecution::Execute(ParentWindow, WindowTitle, Tasks);
+	bool HasSucceeeded = ExecuteTasksAndBroadcastLocalizationTargetDataUpdated(ParentWindow, WindowTitle, Tasks, Target);
 	IFileManager::Get().Delete(*ExportScriptPath); // Don't clutter up the loc config directory with scripts for individual cultures.
 	return HasSucceeeded;
 }
@@ -215,7 +250,7 @@ bool LocalizationCommandletTasks::ImportDialogueScriptForTargets(const TSharedRe
 		Tasks.Add(LocalizationCommandletExecution::FTask(ReportTaskName, ReportScriptPath, ShouldUseProjectFile));
 	}
 
-	return LocalizationCommandletExecution::Execute(ParentWindow, LOCTEXT("ImportDialogueScriptsForAllTargetsWindowTitle", "Import Dialogue Scripts for All Targets"), Tasks);
+	return ExecuteTasksAndBroadcastLocalizationTargetsDataUpdated(ParentWindow, LOCTEXT("ImportDialogueScriptsForAllTargetsWindowTitle", "Import Dialogue Scripts for All Targets"), Tasks, Targets);
 }
 
 bool LocalizationCommandletTasks::ImportDialogueScriptForTarget(const TSharedRef<SWindow>& ParentWindow, ULocalizationTarget* const Target, const TOptional<FString> DirectoryPath)
@@ -235,7 +270,7 @@ bool LocalizationCommandletTasks::ImportDialogueScriptForTarget(const TSharedRef
 	Arguments.Add(TEXT("TargetName"), FText::FromString(Target->Settings.Name));
 
 	const FText WindowTitle = FText::Format(LOCTEXT("ImportDialogueScriptsForTargetWindowTitle", "Import Dialogue Scripts for Target {TargetName}"), Arguments);
-	return LocalizationCommandletExecution::Execute(ParentWindow, WindowTitle, Tasks);
+	return ExecuteTasksAndBroadcastLocalizationTargetDataUpdated(ParentWindow, WindowTitle, Tasks, Target);
 }
 
 bool LocalizationCommandletTasks::ImportDialogueScriptForCulture(const TSharedRef<SWindow>& ParentWindow, ULocalizationTarget* const Target, const FString& CultureName, const TOptional<FString> FilePath)
@@ -264,7 +299,7 @@ bool LocalizationCommandletTasks::ImportDialogueScriptForCulture(const TSharedRe
 
 	const FText WindowTitle = FText::Format(LOCTEXT("ImportDialogueScriptsForCultureForTargetWindowTitle", "Import {CultureName} Dialogue Scripts for Target {TargetName}"), Arguments);
 
-	bool HasSucceeeded = LocalizationCommandletExecution::Execute(ParentWindow, WindowTitle, Tasks);
+	bool HasSucceeeded = ExecuteTasksAndBroadcastLocalizationTargetDataUpdated(ParentWindow, WindowTitle, Tasks, Target);
 	IFileManager::Get().Delete(*ImportScriptPath); // Don't clutter up the loc config directory with scripts for individual cultures.
 	return HasSucceeeded;
 }
@@ -287,7 +322,7 @@ bool LocalizationCommandletTasks::ExportDialogueScriptForTargets(const TSharedRe
 		Tasks.Add(LocalizationCommandletExecution::FTask(ExportTaskName, ExportScriptPath, ShouldUseProjectFile));
 	}
 
-	return LocalizationCommandletExecution::Execute(ParentWindow, LOCTEXT("ExportDialogueScriptsForAllTargetsWindowTitle", "Export Dialogue Scripts for All Targets"), Tasks);
+	return ExecuteTasksAndBroadcastLocalizationTargetsDataUpdated(ParentWindow, LOCTEXT("ExportDialogueScriptsForAllTargetsWindowTitle", "Export Dialogue Scripts for All Targets"), Tasks, Targets);
 }
 
 bool LocalizationCommandletTasks::ExportDialogueScriptForTarget(const TSharedRef<SWindow>& ParentWindow, ULocalizationTarget* const Target, const TOptional<FString> DirectoryPath)
@@ -303,7 +338,7 @@ bool LocalizationCommandletTasks::ExportDialogueScriptForTarget(const TSharedRef
 	Arguments.Add(TEXT("TargetName"), FText::FromString(Target->Settings.Name));
 
 	const FText WindowTitle = FText::Format(LOCTEXT("ExportDialogueScriptsForTargetWindowTitle", "Export Dialogue Scripts for Target {TargetName}"), Arguments);
-	return LocalizationCommandletExecution::Execute(ParentWindow, WindowTitle, Tasks);
+	return ExecuteTasksAndBroadcastLocalizationTargetDataUpdated(ParentWindow, WindowTitle, Tasks, Target);
 }
 
 bool LocalizationCommandletTasks::ExportDialogueScriptForCulture(const TSharedRef<SWindow>& ParentWindow, ULocalizationTarget* const Target, const FString& CultureName, const TOptional<FString> FilePath)
@@ -328,7 +363,7 @@ bool LocalizationCommandletTasks::ExportDialogueScriptForCulture(const TSharedRe
 
 	const FText WindowTitle = FText::Format(LOCTEXT("ExportDialogueScriptsForCultureForTargetWindowTitle", "Export {CultureName} Dialogue Scripts for Target {TargetName}"), Arguments);
 
-	bool HasSucceeeded = LocalizationCommandletExecution::Execute(ParentWindow, WindowTitle, Tasks);
+	bool HasSucceeeded = ExecuteTasksAndBroadcastLocalizationTargetDataUpdated(ParentWindow, WindowTitle, Tasks, Target);
 	IFileManager::Get().Delete(*ExportScriptPath); // Don't clutter up the loc config directory with scripts for individual cultures.
 	return HasSucceeeded;
 }
@@ -350,7 +385,7 @@ bool LocalizationCommandletTasks::ImportDialogueForTargets(const TSharedRef<SWin
 		Tasks.Add(LocalizationCommandletExecution::FTask(ImportDialogueTaskName, ImportDialogueScriptPath, ShouldUseProjectFile));
 	}
 
-	return LocalizationCommandletExecution::Execute(ParentWindow, LOCTEXT("ImportDialogueForAllTargetsWindowTitle", "Import Dialogue for All Targets"), Tasks);
+	return ExecuteTasksAndBroadcastLocalizationTargetsDataUpdated(ParentWindow, LOCTEXT("ImportDialogueForAllTargetsWindowTitle", "Import Dialogue for All Targets"), Tasks, Targets);
 }
 
 bool LocalizationCommandletTasks::ImportDialogueForTarget(const TSharedRef<SWindow>& ParentWindow, ULocalizationTarget* const Target)
@@ -366,7 +401,7 @@ bool LocalizationCommandletTasks::ImportDialogueForTarget(const TSharedRef<SWind
 	Arguments.Add(TEXT("TargetName"), FText::FromString(Target->Settings.Name));
 
 	const FText WindowTitle = FText::Format(LOCTEXT("ImportDialogueForTargetWindowTitle", "Import Dialogue for Target {TargetName}"), Arguments);
-	return LocalizationCommandletExecution::Execute(ParentWindow, WindowTitle, Tasks);
+	return ExecuteTasksAndBroadcastLocalizationTargetDataUpdated(ParentWindow, WindowTitle, Tasks, Target);
 }
 
 bool LocalizationCommandletTasks::ImportDialogueForCulture(const TSharedRef<SWindow>& ParentWindow, ULocalizationTarget* const Target, const FString& CultureName)
@@ -391,7 +426,7 @@ bool LocalizationCommandletTasks::ImportDialogueForCulture(const TSharedRef<SWin
 
 	const FText WindowTitle = FText::Format(LOCTEXT("ImportCultureDialogueForTargetWindowTitle", "Import {CultureName} Dialogue for Target {TargetName}"), Arguments);
 
-	bool HasSucceeeded = LocalizationCommandletExecution::Execute(ParentWindow, WindowTitle, Tasks);
+	bool HasSucceeeded = ExecuteTasksAndBroadcastLocalizationTargetDataUpdated(ParentWindow, WindowTitle, Tasks, Target);
 	IFileManager::Get().Delete(*ImportDialogueScriptPath); // Don't clutter up the loc config directory with scripts for individual cultures.
 	return HasSucceeeded;
 }
@@ -413,7 +448,7 @@ bool LocalizationCommandletTasks::GenerateWordCountReportsForTargets(const TShar
 		Tasks.Add(LocalizationCommandletExecution::FTask(ReportTaskName, ReportScriptPath, ShouldUseProjectFile));
 	}
 
-	return LocalizationCommandletExecution::Execute(ParentWindow, LOCTEXT("GenerateReportsForAllTargetsWindowTitle", "Generate Word Count Reports for All Targets"), Tasks);
+	return ExecuteTasksAndBroadcastLocalizationTargetsDataUpdated(ParentWindow, LOCTEXT("GenerateReportsForAllTargetsWindowTitle", "Generate Word Count Reports for All Targets"), Tasks, Targets);
 }
 
 bool LocalizationCommandletTasks::GenerateWordCountReportForTarget(const TSharedRef<SWindow>& ParentWindow, ULocalizationTarget* const Target)
@@ -429,7 +464,7 @@ bool LocalizationCommandletTasks::GenerateWordCountReportForTarget(const TShared
 	Arguments.Add(TEXT("TargetName"), FText::FromString(Target->Settings.Name));
 
 	const FText WindowTitle = FText::Format(LOCTEXT("GenerateReportForTargetWindowTitle", "Generate Word Count Report for Target {TargetName}"), Arguments);
-	return LocalizationCommandletExecution::Execute(ParentWindow, WindowTitle, Tasks);
+	return ExecuteTasksAndBroadcastLocalizationTargetDataUpdated(ParentWindow, WindowTitle, Tasks, Target);
 }
 
 bool LocalizationCommandletTasks::CompileTextForTargets(const TSharedRef<SWindow>& ParentWindow, const TArray<ULocalizationTarget*>& Targets)
@@ -449,7 +484,7 @@ bool LocalizationCommandletTasks::CompileTextForTargets(const TSharedRef<SWindow
 		Tasks.Add(LocalizationCommandletExecution::FTask(CompileTaskName, CompileScriptPath, ShouldUseProjectFile));
 	}
 
-	return LocalizationCommandletExecution::Execute(ParentWindow, LOCTEXT("GenerateLocResForAllTargetsWindowTitle", "Compile Translations for All Targets"), Tasks);
+	return ExecuteTasksAndBroadcastLocalizationTargetsDataUpdated(ParentWindow, LOCTEXT("GenerateLocResForAllTargetsWindowTitle", "Compile Translations for All Targets"), Tasks, Targets);
 }
 
 bool LocalizationCommandletTasks::CompileTextForTarget(const TSharedRef<SWindow>& ParentWindow, ULocalizationTarget* const Target)
@@ -465,7 +500,7 @@ bool LocalizationCommandletTasks::CompileTextForTarget(const TSharedRef<SWindow>
 	Arguments.Add(TEXT("TargetName"), FText::FromString(Target->Settings.Name));
 
 	const FText WindowTitle = FText::Format(LOCTEXT("GenerateLocResForTargetWindowTitle", "Compile Translations for Target {TargetName}"), Arguments);
-	return LocalizationCommandletExecution::Execute(ParentWindow, WindowTitle, Tasks);
+	return ExecuteTasksAndBroadcastLocalizationTargetDataUpdated(ParentWindow, WindowTitle, Tasks, Target);
 }
 
 bool LocalizationCommandletTasks::CompileTextForCulture(const TSharedRef<SWindow>& ParentWindow, ULocalizationTarget* const Target, const FString& CultureName)
@@ -490,7 +525,7 @@ bool LocalizationCommandletTasks::CompileTextForCulture(const TSharedRef<SWindow
 
 	const FText WindowTitle = FText::Format(LOCTEXT("CompileCultureForTargetWindowTitle", "Compile {CultureName} Translations for Target {TargetName}"), Arguments);
 
-	bool HasSucceeeded = LocalizationCommandletExecution::Execute(ParentWindow, WindowTitle, Tasks);
+	bool HasSucceeeded = ExecuteTasksAndBroadcastLocalizationTargetDataUpdated(ParentWindow, WindowTitle, Tasks, Target);
 	IFileManager::Get().Delete(*CompileScriptPath); // Don't clutter up the loc config directory with scripts for individual cultures.
 	return HasSucceeeded;
 }
