Fix missing scrollbars from localization dialog

* Vertical and Horizontal ScrollBars need to be constructed and assigned to temporaries before use.

diff --git a/Engine/Source/Editor/LocalizationCommandletExecution/Private/LocalizationCommandletExecution.cpp b/Engine/Source/Editor/LocalizationCommandletExecution/Private/LocalizationCommandletExecution.cpp
--- a/Engine/Source/Editor/LocalizationCommandletExecution/Private/LocalizationCommandletExecution.cpp
+++ b/Engine/Source/Editor/LocalizationCommandletExecution/Private/LocalizationCommandletExecution.cpp
@@ -149,8 +149,10 @@ namespace
 			TaskListModels.Add(Model);
 		}
 
-		TSharedPtr<SScrollBar> VerticalScrollBar;
-		TSharedPtr<SScrollBar> HorizontalScrollBar;
+		TSharedRef<SScrollBar> VerticalScrollBar = SNew(SScrollBar)
+			.Orientation(Orient_Vertical);
+		TSharedRef<SScrollBar> HorizontalScrollBar = SNew(SScrollBar)
+			.Orientation(Orient_Horizontal);
 
 		ChildSlot
 			[
@@ -225,15 +227,13 @@ namespace
 								+SVerticalBox::Slot()
 									.AutoHeight()
 									[
-										SAssignNew(HorizontalScrollBar, SScrollBar)
-										.Orientation(EOrientation::Orient_Horizontal)
+										HorizontalScrollBar
 									]
 							]
 							+SHorizontalBox::Slot()
 								.AutoWidth()
 								[
-									SAssignNew(VerticalScrollBar, SScrollBar)
-									.Orientation(EOrientation::Orient_Vertical)
+									VerticalScrollBar
 								]
 						]
 					]
