Fixed issue with GatherText commandlet not executing properly when launched from the Localization Dashboard

diff --git a/Engine/Source/Editor/LocalizationCommandletExecution/Private/LocalizationCommandletExecution.cpp b/Engine/Source/Editor/LocalizationCommandletExecution/Private/LocalizationCommandletExecution.cpp
--- a/Engine/Source/Editor/LocalizationCommandletExecution/Private/LocalizationCommandletExecution.cpp
+++ b/Engine/Source/Editor/LocalizationCommandletExecution/Private/LocalizationCommandletExecution.cpp
@@ -757,6 +757,8 @@ TSharedPtr<FLocalizationCommandletProcess> FLocalizationCommandletProcess::Execu
 		}
 	}
 
+	CommandletArguments += TEXT(" -Unattended");
+
 	const FString ProjectFilePath = FString::Printf(TEXT("\"%s\""), *FPaths::ConvertRelativePathToFull(FPaths::GetProjectFilePath()));
 	const FString ProcessArguments = CommandletHelpers::BuildCommandletProcessArguments(TEXT("GatherText"), UseProjectFile ? *ProjectFilePath : nullptr, *CommandletArguments);
 	FProcHandle CommandletProcessHandle = FPlatformProcess::CreateProc(*FUnrealEdMisc::Get().GetExecutableForCommandlets(), *ProcessArguments, true, true, true, nullptr, 0, nullptr, WritePipe);
